Update QAGame content for rename of EditorMenus module to ToolsMenu module
#rnx

diff --git a/Engine/Source/Runtime/ToolMenus/Public/ToolMenus.h b/Engine/Source/Runtime/ToolMenus/Public/ToolMenus.h
--- a/Engine/Source/Runtime/ToolMenus/Public/ToolMenus.h
+++ b/Engine/Source/Runtime/ToolMenus/Public/ToolMenus.h
@@ -91,6 +91,7 @@ public:
 
 	UToolMenus();
 
+	UFUNCTION(BlueprintCallable, Category = "Tool Menus")
 	static UToolMenus* Get();
 
 	/** Try to get UToolMenus without forcing ToolMenus module to load. */
